Performance Capture

- Move EnhancedInput LiveLink plugin into engine, under Virtual Production

#virtualized

diff --git a/Engine/Plugins/VirtualProduction/EnhancedInputLiveLink/EnhancedInputLiveLink.uplugin b/Engine/Plugins/VirtualProduction/EnhancedInputLiveLink/EnhancedInputLiveLink.uplugin
new file mode 100644
--- /dev/null
+++ b/Engine/Plugins/VirtualProduction/EnhancedInputLiveLink/EnhancedInputLiveLink.uplugin
@@ -0,0 +1,34 @@
+﻿{
+	"FileVersion": 3,
+	"Version": 1,
+	"VersionName": "1.0",
+	"FriendlyName": "EnhancedInputLiveLink",
+	"Description": "A lightweight plugin to enable Enhanced Input within a Virtual Live Link subject",
+	"Category": "Virtual Production",
+	"CreatedBy": "Epic Games Inc.",
+	"CreatedByURL": "Virtual Production",
+	"DocsURL": "",
+	"MarketplaceURL": "",
+	"CanContainContent": true,
+	"IsBetaVersion": true,
+	"IsExperimentalVersion": true,
+	"Installed": false,
+	"Modules": [
+		{
+			"Name": "EnhancedInputLiveLink",
+			"Type": "Editor",
+			"LoadingPhase": "Default"
+		}
+	],
+	"Plugins": [
+		{
+			"Name": "EnhancedInput",
+			"Enabled": true
+		},
+		{
+			"Name": "LiveLink",
+			"Enabled": true
+		}
+	],
+	"SupportURL": ""
+}
\ No newline at end of file
diff --git a/Engine/Plugins/VirtualProduction/EnhancedInputLiveLink/Source/EnhancedInputLiveLink/EnhancedInputLiveLink.Build.cs b/Engine/Plugins/VirtualProduction/EnhancedInputLiveLink/Source/EnhancedInputLiveLink/EnhancedInputLiveLink.Build.cs
new file mode 100644
--- /dev/null
+++ b/Engine/Plugins/VirtualProduction/EnhancedInputLiveLink/Source/EnhancedInputLiveLink/EnhancedInputLiveLink.Build.cs
@@ -0,0 +1,56 @@
+﻿// Copyright Epic Games, Inc. All Rights Reserved.
+
+using UnrealBuildTool;
+
+public class EnhancedInputLiveLink : ModuleRules
+{
+	public EnhancedInputLiveLink(ReadOnlyTargetRules Target) : base(Target)
+	{
+		PCHUsage = ModuleRules.PCHUsageMode.UseExplicitOrSharedPCHs;
+		
+		PublicIncludePaths.AddRange(
+			new string[] {
+				// ... add public include paths required here ...
+			}
+			);
+				
+		
+		PrivateIncludePaths.AddRange(
+			new string[] {
+				// ... add other private include paths required here ...
+			}
+			);
+			
+		
+		PublicDependencyModuleNames.AddRange(
+			new string[]
+			{
+				"Core",
+				"CoreUObject", 
+				"Engine", 
+				"InputCore",
+				"LiveLinkInterface",
+				"LiveLinkComponents"
+				// ... add other public dependencies that you statically link with here ...
+			}
+			);
+			
+		
+		PrivateDependencyModuleNames.AddRange(
+			new string[] 
+			{ 
+				"LiveLink",
+				"InputCore",
+				"EnhancedInput",
+				"Slate",
+				"SlateCore",
+				"KismetCompiler",
+				"BlueprintGraph",
+				"EditorFramework",
+				"UnrealEd",
+				"InputEditor",
+				"LiveLinkEditor",
+				"EditorSubsystem",
+			});
+	}
+}
diff --git a/Engine/Plugins/VirtualProduction/EnhancedInputLiveLink/Source/EnhancedInputLiveLink/Private/EnhancedInputLiveLink.cpp b/Engine/Plugins/VirtualProduction/EnhancedInputLiveLink/Source/EnhancedInputLiveLink/Private/EnhancedInputLiveLink.cpp
new file mode 100644
--- /dev/null
+++ b/Engine/Plugins/VirtualProduction/EnhancedInputLiveLink/Source/EnhancedInputLiveLink/Private/EnhancedInputLiveLink.cpp
@@ -0,0 +1,20 @@
+﻿// Copyright Epic Games, Inc. All Rights Reserved.
+
+#include "../Public/EnhancedInputLiveLink.h"
+
+#define LOCTEXT_NAMESPACE "FEnhancedInputLiveLinkModule"
+
+void FEnhancedInputLiveLinkModule::StartupModule()
+{
+	// This code will execute after your module is loaded into memory; the exact timing is specified in the .uplugin file per-module
+}
+
+void FEnhancedInputLiveLinkModule::ShutdownModule()
+{
+	// This function may be called during shutdown to clean up your module.  For modules that support dynamic reloading,
+	// we call this function before unloading the module.
+}
+
+#undef LOCTEXT_NAMESPACE
+	
+IMPLEMENT_MODULE(FEnhancedInputLiveLinkModule, EnhancedInputLiveLink)
\ No newline at end of file
diff --git a/Engine/Plugins/VirtualProduction/EnhancedInputLiveLink/Source/EnhancedInputLiveLink/Private/EnhancedInputVirtualSubject.cpp b/Engine/Plugins/VirtualProduction/EnhancedInputLiveLink/Source/EnhancedInputLiveLink/Private/EnhancedInputVirtualSubject.cpp
new file mode 100644
--- /dev/null
+++ b/Engine/Plugins/VirtualProduction/EnhancedInputLiveLink/Source/EnhancedInputLiveLink/Private/EnhancedInputVirtualSubject.cpp
@@ -0,0 +1,137 @@
+// Copyright Epic Games, Inc. All Rights Reserved.
+
+#include "EnhancedInputVirtualSubject.h"
+
+#include "EnhancedInputSubsystems.h"
+#include "InputMappingContext.h"
+#include "Engine/InputDelegateBinding.h"
+#include "Engine/LocalPlayer.h"
+#include "GameFramework/InputSettings.h"
+#include "GameFramework/PlayerController.h"
+#if WITH_EDITOR
+#	include "Editor.h"
+#	include "EnhancedInputEditorSubsystem.h"
+#endif
+
+void UEnhancedInputVirtualSubject::Initialize(FLiveLinkSubjectKey InSubjectKey, TSubclassOf<ULiveLinkRole> InRole, ILiveLinkClient* InLiveLinkClient)
+{
+	Super::Initialize(InSubjectKey, Role, LiveLinkClient);
+
+	{
+		FEditorScriptExecutionGuard ScriptGuard;
+		
+		CreateEditorInput();
+	}
+}
+
+
+
+void UEnhancedInputVirtualSubject::CreateEditorInput()
+{
+	InputComponent = NewObject<UInputComponent>(this, UInputSettings::GetDefaultInputComponentClass(), TEXT("LivelinkEnhancedInputComponent0"), RF_Transient);
+	UInputDelegateBinding::BindInputDelegates(GetClass(), InputComponent, this);
+	
+	RegisterInputComponent();
+}
+
+void UEnhancedInputVirtualSubject::RemoveEditorInput()
+{
+	if (InputComponent)
+		UnregisterInputComponent();
+	{
+		InputComponent->DestroyComponent();
+	}
+	InputComponent = nullptr;
+
+}
+
+
+void UEnhancedInputVirtualSubject::RegisterInputComponent()
+{
+	// Ensure we start from a clean slate
+	UnregisterInputComponent();
+	
+	if (const UWorld* World = GetWorld(); IsValid(World) && World->IsGameWorld())
+	{
+		if (APlayerController* PC = World->GetFirstPlayerController())
+		{
+			PC->PushInputComponent(InputComponent);
+			bIsInputRegistered = true;
+		}
+	}
+#if WITH_EDITOR
+	else if (GEditor)
+	{
+		if (UEnhancedInputEditorSubsystem* EditorInputSubsystem = GEditor->GetEditorSubsystem<UEnhancedInputEditorSubsystem>())
+		{
+			EditorInputSubsystem->StartConsumingInput();
+			EditorInputSubsystem->PushInputComponent(InputComponent);
+			const int32 InPriority = 1;
+			UInputMappingContext* NewMappingContext = Cast<UInputMappingContext>(MappingContext.TryLoad());
+			if(NewMappingContext == nullptr)
+			{
+				// early out if there is no Mapping Context reference available
+				return;
+			}
+			
+			if (IEnhancedInputSubsystemInterface* EnhancedInputSubsystemInterface = GetEnhancedInputSubsystemInterface())
+			{
+				if(IsValid(NewMappingContext))
+				{
+					EnhancedInputSubsystemInterface->AddMappingContext(NewMappingContext, InPriority);
+					
+				}
+			}
+			bIsInputRegistered = true;
+			
+		}
+	}
+#endif
+}
+
+
+void UEnhancedInputVirtualSubject::UnregisterInputComponent()
+{
+	// Removes the component from both editor and runtime input systems if possible
+	if (const UWorld* World = GetWorld(); IsValid(World) && World->IsGameWorld())
+	{
+		if (APlayerController* PC = World->GetFirstPlayerController())
+		{
+			PC->PopInputComponent(InputComponent);
+			
+		}
+	}
+#if WITH_EDITOR
+	if (GEditor)
+	{
+		if (UEnhancedInputEditorSubsystem* EditorInputSubsystem = GEditor->GetEditorSubsystem<UEnhancedInputEditorSubsystem>())
+		{
+			EditorInputSubsystem->PopInputComponent(InputComponent);
+			EditorInputSubsystem->StopConsumingInput();
+		}
+			
+	}
+#endif
+
+	bIsInputRegistered = false;
+}
+
+IEnhancedInputSubsystemInterface* UEnhancedInputVirtualSubject::GetEnhancedInputSubsystemInterface() const
+{
+	if (const UWorld* World = GetWorld(); IsValid(World) && World->IsGameWorld())
+	{
+		if (const ULocalPlayer* FirstLocalPlayer = World->GetFirstLocalPlayerFromController())
+		{
+			return ULocalPlayer::GetSubsystem<UEnhancedInputLocalPlayerSubsystem>(FirstLocalPlayer);
+		}
+	}
+#if WITH_EDITOR
+	else if (GEditor)
+	{
+		return GEditor->GetEditorSubsystem<UEnhancedInputEditorSubsystem>();
+	}
+#endif
+	return nullptr;
+}
+
+
diff --git a/Engine/Plugins/VirtualProduction/EnhancedInputLiveLink/Source/EnhancedInputLiveLink/Private/EnhancedInputVirtualSubjectFactory.cpp b/Engine/Plugins/VirtualProduction/EnhancedInputLiveLink/Source/EnhancedInputLiveLink/Private/EnhancedInputVirtualSubjectFactory.cpp
new file mode 100644
--- /dev/null
+++ b/Engine/Plugins/VirtualProduction/EnhancedInputLiveLink/Source/EnhancedInputLiveLink/Private/EnhancedInputVirtualSubjectFactory.cpp
@@ -0,0 +1,22 @@
+﻿// Copyright Epic Games, Inc. All Rights Reserved.
+
+#include "EnhancedInputVirtualSubjectFactory.h"
+#include "EnhancedInputVirtualSubject.h"
+
+#include UE_INLINE_GENERATED_CPP_BY_NAME(EnhancedInputVirtualSubjectFactory)
+
+#define LOCTEXT_NAMESPACE "EnhancedInputVirtualSubjectFactory"
+
+
+UEnhancedInputVirtualSubjectFactory::UEnhancedInputVirtualSubjectFactory()
+{
+	SupportedClass = UBlueprint::StaticClass();
+	ParentClass = UEnhancedInputVirtualSubject::StaticClass();
+}
+
+FText UEnhancedInputVirtualSubjectFactory::GetDisplayName() const
+{
+	return LOCTEXT("EnhancedInputVirtualSubjectName", "Enhanced Input Virtual Subject");
+}
+
+#undef LOCTEXT_NAMESPACE
\ No newline at end of file
diff --git a/Engine/Plugins/VirtualProduction/EnhancedInputLiveLink/Source/EnhancedInputLiveLink/Public/EnhancedInputLiveLink.h b/Engine/Plugins/VirtualProduction/EnhancedInputLiveLink/Source/EnhancedInputLiveLink/Public/EnhancedInputLiveLink.h
new file mode 100644
--- /dev/null
+++ b/Engine/Plugins/VirtualProduction/EnhancedInputLiveLink/Source/EnhancedInputLiveLink/Public/EnhancedInputLiveLink.h
@@ -0,0 +1,15 @@
+﻿// Copyright Epic Games, Inc. All Rights Reserved.
+
+#pragma once
+
+#include "CoreMinimal.h"
+#include "Modules/ModuleManager.h"
+
+class FEnhancedInputLiveLinkModule : public IModuleInterface
+{
+public:
+
+	/** IModuleInterface implementation */
+	virtual void StartupModule() override;
+	virtual void ShutdownModule() override;
+};
diff --git a/Engine/Plugins/VirtualProduction/EnhancedInputLiveLink/Source/EnhancedInputLiveLink/Public/EnhancedInputVirtualSubject.h b/Engine/Plugins/VirtualProduction/EnhancedInputLiveLink/Source/EnhancedInputLiveLink/Public/EnhancedInputVirtualSubject.h
new file mode 100644
--- /dev/null
+++ b/Engine/Plugins/VirtualProduction/EnhancedInputLiveLink/Source/EnhancedInputLiveLink/Public/EnhancedInputVirtualSubject.h
@@ -0,0 +1,56 @@
+// Copyright Epic Games, Inc. All Rights Reserved.
+
+#pragma once
+
+#include "InputCoreTypes.h"
+#include "EnhancedInputSubsystemInterface.h"
+#include "VirtualSubjects/LiveLinkBlueprintVirtualSubject.h"
+
+#include "EnhancedInputVirtualSubject.generated.h"
+
+class UInputComponent;
+class UInputMappingContext;
+class IEnhancedInputSubsystemInterface;
+class UInputMappingContext;
+class UMotionControllerComponent;
+
+/**
+ * 
+ */
+UCLASS(Blueprintable, Abstract)
+class ENHANCEDINPUTLIVELINK_API UEnhancedInputVirtualSubject : public ULiveLinkBlueprintVirtualSubject
+{
+	GENERATED_BODY()
+
+public:
+	
+	UPROPERTY(EditAnywhere, Category="Enhanced Input LiveLink", meta=(AllowedClasses="/Script/EnhancedInput.InputMappingContext"))
+	FSoftObjectPath MappingContext;
+
+	/** Returns the current InputComponent. */
+	UFUNCTION(BlueprintCallable, BlueprintPure, Category = "Input|Editor")
+	UInputComponent* GetInputComponent() const
+	{
+		return InputComponent.Get();
+	}
+	
+	void RegisterInputComponent();
+	void UnregisterInputComponent();
+
+protected:
+	bool bIsInputRegistered = false;
+	IEnhancedInputSubsystemInterface* GetEnhancedInputSubsystemInterface() const;
+	
+private:
+	
+	UPROPERTY(Transient, DuplicateTransient)
+	TObjectPtr<UInputComponent> InputComponent;
+	
+	virtual void Initialize(FLiveLinkSubjectKey InSubjectKey, TSubclassOf<ULiveLinkRole> InRole, ILiveLinkClient* InLiveLinkClient) override;
+	/** Creates the InputComponent if it does not already exist and registers all sub-object callbacks to it */
+	void CreateEditorInput();
+
+	/** Removes the InputComponent from this object */
+	void RemoveEditorInput();
+		
+};
diff --git a/Engine/Plugins/VirtualProduction/EnhancedInputLiveLink/Source/EnhancedInputLiveLink/Public/EnhancedInputVirtualSubjectFactory.h b/Engine/Plugins/VirtualProduction/EnhancedInputLiveLink/Source/EnhancedInputLiveLink/Public/EnhancedInputVirtualSubjectFactory.h
new file mode 100644
--- /dev/null
+++ b/Engine/Plugins/VirtualProduction/EnhancedInputLiveLink/Source/EnhancedInputLiveLink/Public/EnhancedInputVirtualSubjectFactory.h
@@ -0,0 +1,23 @@
+﻿// Copyright Epic Games, Inc. All Rights Reserved.
+
+#pragma once
+
+#include "VirtualSubjects/LiveLinkBlueprintVirtualSubjectFactory.h"
+#include "EnhancedInputVirtualSubjectFactory.generated.h"
+
+/**
+ * 
+ */
+UCLASS()
+class ENHANCEDINPUTLIVELINK_API UEnhancedInputVirtualSubjectFactory : public ULiveLinkBlueprintVirtualSubjectFactory
+{
+	GENERATED_BODY()
+
+public:
+
+	UEnhancedInputVirtualSubjectFactory();
+
+	//~ Begin UFactory Interface
+	virtual FText GetDisplayName() const override;
+	//~ Begin UFactory Interface
+};
